GitHub pull request #653

diff --git a/Engine/Source/Runtime/SQLiteSupport/Private/SQLiteDatabaseConnection.cpp b/Engine/Source/Runtime/SQLiteSupport/Private/SQLiteDatabaseConnection.cpp
--- a/Engine/Source/Runtime/SQLiteSupport/Private/SQLiteDatabaseConnection.cpp
+++ b/Engine/Source/Runtime/SQLiteSupport/Private/SQLiteDatabaseConnection.cpp
@@ -7,7 +7,7 @@ bool FSQLiteDatabase::Execute(const TCHAR* CommandString, FSQLiteResultSet*& Rec
 	
 	//Compile the statement/query
 	sqlite3_stmt* PreparedStatement;
-	int32 PrepareStatus = sqlite3_prepare16_v2(DbHandle, CommandString, -1, &PreparedStatement, NULL);
+	int32 PrepareStatus = sqlite3_prepare_v2(DbHandle, TCHAR_TO_UTF8(CommandString), -1, &PreparedStatement, NULL);
 	if (PrepareStatus == SQLITE_OK)
 	{
 		//Initialize records from compiled query
@@ -26,7 +26,7 @@ bool FSQLiteDatabase::Execute(const TCHAR* CommandString)
 {
 	//Compile the statement/query
 	sqlite3_stmt* PreparedStatement;
-	int32 PrepareStatus = sqlite3_prepare16_v2(DbHandle, CommandString, -1, &PreparedStatement, NULL);
+	int32 PrepareStatus = sqlite3_prepare_v2(DbHandle, TCHAR_TO_UTF8(CommandString), -1, &PreparedStatement, NULL);
 	if (PrepareStatus == SQLITE_OK)
 	{
 		int32 StepStatus = SQLITE_ERROR;
@@ -74,14 +74,14 @@ bool FSQLiteDatabase::Open(const TCHAR* ConnectionString, const TCHAR* RemoteCon
 		return false;
 	}
 
-	int32 Result = sqlite3_open16(ConnectionString, &DbHandle);
+	int32 Result = sqlite3_open(TCHAR_TO_UTF8(ConnectionString), &DbHandle);
 	return Result == SQLITE_OK;
 }
 
 FString FSQLiteDatabase::GetLastError()
 {
 	TCHAR* ErrorString = NULL;
-	ErrorString = (TCHAR*) sqlite3_errmsg16(DbHandle);
+	ErrorString = (TCHAR*) UTF8_TO_TCHAR(sqlite3_errmsg(DbHandle));
 	if (ErrorString)
 	{
 		return FString(ErrorString);
diff --git a/Engine/Source/Runtime/SQLiteSupport/Private/SQLiteResultSet.cpp b/Engine/Source/Runtime/SQLiteSupport/Private/SQLiteResultSet.cpp
--- a/Engine/Source/Runtime/SQLiteSupport/Private/SQLiteResultSet.cpp
+++ b/Engine/Source/Runtime/SQLiteSupport/Private/SQLiteResultSet.cpp
@@ -85,7 +85,7 @@ float FSQLiteResultSet::GetFloat(const TCHAR* Column) const
 	}
 	else
 	{
-		return FCString::Atof((const wchar_t *)sqlite3_column_text16(PreparedQuery, ColumnIndex));
+		return FCString::Atof(UTF8_TO_TCHAR(sqlite3_column_text(PreparedQuery, ColumnIndex)));
 	}
 
 
@@ -101,7 +101,7 @@ int32 FSQLiteResultSet::GetInt(const TCHAR* Column) const
 	}
 	else
 	{
-		return FCString::Atoi((const wchar_t *)sqlite3_column_text16(PreparedQuery, ColumnIndex));
+		return FCString::Atoi(UTF8_TO_TCHAR(sqlite3_column_text(PreparedQuery, ColumnIndex)));
 	}
 
 
@@ -117,7 +117,7 @@ FString FSQLiteResultSet::GetString(const TCHAR* Column) const
 	}
 	else
 	{
-		return FString((const wchar_t *)sqlite3_column_text16(PreparedQuery, ColumnIndex));
+		return FString(UTF8_TO_TCHAR(sqlite3_column_text(PreparedQuery, ColumnIndex)));
 	}
 
 }
@@ -142,3 +142,4 @@ void FSQLiteResultSet::MoveToFirst()
 	sqlite3_reset(PreparedQuery);
 	StepStatus = sqlite3_step(PreparedQuery);
 }
+
diff --git a/Engine/Source/Runtime/SQLiteSupport/SQLiteSupport.Build.cs b/Engine/Source/Runtime/SQLiteSupport/SQLiteSupport.Build.cs
--- a/Engine/Source/Runtime/SQLiteSupport/SQLiteSupport.Build.cs
+++ b/Engine/Source/Runtime/SQLiteSupport/SQLiteSupport.Build.cs
@@ -20,7 +20,15 @@ namespace UnrealBuildTool.Rules
 					PlatformName = "x64/";
 					break;
 			
-				// This needs to be extended for Linux and Mac.
+			    case UnrealTargetPlatform.IOS:
+                    PlatformName = "IOS/";
+                    break;
+                case UnrealTargetPlatform.Mac:
+                    PlatformName = "Mac/";
+                    break;
+                case UnrealTargetPlatform.Linux:
+                    PlatformName = "Linux/";
+                    break;
 			}
 
 			switch (Target.Configuration)
@@ -37,7 +45,7 @@ namespace UnrealBuildTool.Rules
 			}
 		
 			string LibraryPath = "" + UEBuildConfiguration.UEThirdPartySourceDirectory + "sqlite/lib/" + PlatformName + ConfigurationName;
-			string LibraryFilename = Path.Combine(LibraryPath, "sqlite.lib");
+			string LibraryFilename = Path.Combine(LibraryPath, "sqlite" + UEBuildPlatform.GetBuildPlatform(Target.Platform).GetBinaryExtension(UEBuildBinaryType.StaticLibrary));
 			if (!File.Exists(LibraryFilename))
 			{
 				throw new BuildException("Please refer to the Engine/Source/ThirdParty/sqlite/README.txt file prior to enabling this module.");
