SQLiteSupport fix. Removed private PCH inclusion in public header.

diff --git a/Engine/Source/Runtime/SQLiteSupport/Public/SQLiteDatabaseConnection.h b/Engine/Source/Runtime/SQLiteSupport/Public/SQLiteDatabaseConnection.h
--- a/Engine/Source/Runtime/SQLiteSupport/Public/SQLiteDatabaseConnection.h
+++ b/Engine/Source/Runtime/SQLiteSupport/Public/SQLiteDatabaseConnection.h
@@ -1,9 +1,10 @@
 // Copyright 1998-2015 Epic Games, Inc. All Rights Reserved.
 
 #pragma once
-#include "SQLiteSupportPrivatePCH.h"
 
+#include "SQLiteSupport.h"
 #include "SQLiteResultSet.h"
+
 /**
 * SQLite database file
 */
